Fix missing semicolon in GooglePAD_APL.xml

#ue4
#android

diff --git a/Engine/Plugins/Runtime/GooglePAD/Source/GooglePAD/GooglePAD_APL.xml b/Engine/Plugins/Runtime/GooglePAD/Source/GooglePAD/GooglePAD_APL.xml
--- a/Engine/Plugins/Runtime/GooglePAD/Source/GooglePAD/GooglePAD_APL.xml
+++ b/Engine/Plugins/Runtime/GooglePAD/Source/GooglePAD/GooglePAD_APL.xml
@@ -256,7 +256,7 @@ import com.google.android.play.core.tasks.Task;
 			Log.debug("Found obbassets asset pack");
 			AssetLocation obbLocation = assetPackManager.getAssetLocation("obbassets", "main.obb.png");
 			if (obbLocation != null) {
-				String mainpngPath = obbLocation.path()
+				String mainpngPath = obbLocation.path();
 				Log.debug("Found main.obb.png asset: " + mainpngPath);
 				if (mainpngPath.endsWith(".apk"))
 				{
