Fix GooglePAD disable cases

#android

diff --git a/Engine/Plugins/Runtime/GooglePAD/Source/GooglePAD/GooglePAD_APL.xml b/Engine/Plugins/Runtime/GooglePAD/Source/GooglePAD/GooglePAD_APL.xml
--- a/Engine/Plugins/Runtime/GooglePAD/Source/GooglePAD/GooglePAD_APL.xml
+++ b/Engine/Plugins/Runtime/GooglePAD/Source/GooglePAD/GooglePAD_APL.xml
@@ -51,12 +51,16 @@
 				<if condition="bEnabled">
 					<true>
 						<if condition="bPackageDataInsideApk">
-							<log text="Disabled because Package Data Inside Apk enabled; disable to use GooglePAD for delivery"/>
-							<setBool result="bEnabled" value="false"/>
+							<true>
+								<log text="Disabled because Package Data Inside Apk enabled; disable to use GooglePAD for delivery"/>
+								<setBool result="bEnabled" value="false"/>
+							</true>
 						</if>
 						<if condition="bAllowPatchOBBFile">
-							<log text="Disabled because Allow patch OBB file enabled; disable to use GooglePAD for delivery"/>
-							<setBool result="bEnabled" value="false"/>
+							<true>
+								<log text="Disabled because Allow patch OBB file enabled; disable to use GooglePAD for delivery"/>
+								<setBool result="bEnabled" value="false"/>
+							</true>
 						</if>
 					</true>
 				</if>
@@ -246,6 +250,15 @@ import com.google.android.play.core.tasks.Task;
 	}
 </insert>
 			</true>
+			<false>
+				<insert>
+	// for GooglePAD
+	public boolean AndroidThunkJava_GooglePAD_Available()
+	{
+		return false;
+	}
+</insert>
+			</false>
 		</if>
 	</gameActivityClassAdditions>
 
