Fixed Android builds with NDK 25

diff --git a/Engine/Plugins/Runtime/GooglePAD/Source/GooglePAD/GooglePAD.Build.cs b/Engine/Plugins/Runtime/GooglePAD/Source/GooglePAD/GooglePAD.Build.cs
--- a/Engine/Plugins/Runtime/GooglePAD/Source/GooglePAD/GooglePAD.Build.cs
+++ b/Engine/Plugins/Runtime/GooglePAD/Source/GooglePAD/GooglePAD.Build.cs
@@ -57,13 +57,6 @@ public class GooglePAD : ModuleRules
 			string PlayCoreLibPath = Path.Combine(PlayCoreSDKPath, "libs");
 
 			string NDKVersion = "ndk21.4.7075529";
-			IAndroidToolChain TempToolchain = AndroidExports.CreateTempToolChain();
-			UInt64 NdkVersionInt = TempToolchain.GetNdkVersionInt();
-			if (NdkVersionInt >= 230000)
-			{
-				NDKVersion = "ndk23.0.7599858";
-			}
-
 			PublicAdditionalLibraries.Add(Path.Combine(PlayCoreLibPath, "arm64-v8a", NDKVersion, "c++_static", "libplaycore.so"));
 			PublicAdditionalLibraries.Add(Path.Combine(PlayCoreLibPath, "x86_64", NDKVersion, "c++_static", "libplaycore.so"));
 
