Statically link with Play Core SDK for Android to avoid conflict with libc++ symbols that libplaycore.so re-exports

diff --git a/Engine/Plugins/Runtime/GooglePAD/Source/GooglePAD/GooglePAD.Build.cs b/Engine/Plugins/Runtime/GooglePAD/Source/GooglePAD/GooglePAD.Build.cs
--- a/Engine/Plugins/Runtime/GooglePAD/Source/GooglePAD/GooglePAD.Build.cs
+++ b/Engine/Plugins/Runtime/GooglePAD/Source/GooglePAD/GooglePAD.Build.cs
@@ -57,8 +57,8 @@ public class GooglePAD : ModuleRules
 			string PlayCoreLibPath = Path.Combine(PlayCoreSDKPath, "libs");
 
 			string NDKVersion = "ndk21.4.7075529";
-			PublicAdditionalLibraries.Add(Path.Combine(PlayCoreLibPath, "arm64-v8a", NDKVersion, "c++_static", "libplaycore.so"));
-			PublicAdditionalLibraries.Add(Path.Combine(PlayCoreLibPath, "x86_64", NDKVersion, "c++_static", "libplaycore.so"));
+			PublicAdditionalLibraries.Add(Path.Combine(PlayCoreLibPath, "arm64-v8a", NDKVersion, "c++_static", "libplaycore_static.a"));
+			PublicAdditionalLibraries.Add(Path.Combine(PlayCoreLibPath, "x86_64", NDKVersion, "c++_static", "libplaycore_static.a"));
 
 			PrivateIncludePaths.Add(PlayCoreIncludePath);
 		}
diff --git a/Engine/Plugins/Runtime/GooglePAD/Source/GooglePAD/GooglePAD_APL.xml b/Engine/Plugins/Runtime/GooglePAD/Source/GooglePAD/GooglePAD_APL.xml
--- a/Engine/Plugins/Runtime/GooglePAD/Source/GooglePAD/GooglePAD_APL.xml
+++ b/Engine/Plugins/Runtime/GooglePAD/Source/GooglePAD/GooglePAD_APL.xml
@@ -9,9 +9,6 @@
 		<setBoolFromProperty result="bEnableBundle" ini="Engine" section="/Script/AndroidRuntimeSettings.AndroidRuntimeSettings" property="bEnableBundle" default="false"/>
 		<setBoolFromProperty result="bAllowPatchOBBFile" ini="Engine" section="/Script/AndroidRuntimeSettings.AndroidRuntimeSettings" property="bAllowPatchOBBFile" default="false"/>
 
-		<!-- NDK path -->
-		<setString result="NDKVersion" value="ndk21.4.7075529"/>
-
 		<if condition="bEnabled">
 			<true>
 				<!-- disable if app bundle disabled -->
@@ -71,11 +68,8 @@
 	</init>
 
 	<resourceCopies>
-		<log text="Copying libplaycore.so and proguard files"/>
-		<!-- note: have to stage this since we linked it -->
-		<copyFile src="$S(AbsPluginDir)/../ThirdParty/play-core-native-sdk/libs/$S(Architecture)/$S(NDKVersion)/c++_static/libplaycore.so"
-					dst="$S(BuildDir)/libs/$S(Architecture)/libplaycore.so" />
 		<isDistribution>
+			<log text="Copying proguard files for Play Core SDK"/>
 			<copyDir src="$S(AbsPluginDir)/../ThirdParty/play-core-native-sdk/proguard"
 					dst="$S(BuildDir)/gradle/app/proguard" />
 		</isDistribution>
