Update SDK version

#android

diff --git a/Engine/Config/Android/Android_SDK.json b/Engine/Config/Android/Android_SDK.json
--- a/Engine/Config/Android/Android_SDK.json
+++ b/Engine/Config/Android/Android_SDK.json
@@ -3,10 +3,10 @@
 	"MinVersion": "r25b",
 	"MaxVersion": "r27",
 
-	"AutoSDKDirectory": "-26",
+	"AutoSDKDirectory": "-27",
 
-	"platforms": "android-33",
-	"build-tools": "33.0.1",
+	"platforms": "android-34",
+	"build-tools": "34.0.0",
 	"cmake": "3.22.1",
 	"ndk": "25.1.8937393",
 }
diff --git a/Engine/Plugins/Runtime/GooglePAD/Source/GooglePAD/GooglePAD_APL.xml b/Engine/Plugins/Runtime/GooglePAD/Source/GooglePAD/GooglePAD_APL.xml
--- a/Engine/Plugins/Runtime/GooglePAD/Source/GooglePAD/GooglePAD_APL.xml
+++ b/Engine/Plugins/Runtime/GooglePAD/Source/GooglePAD/GooglePAD_APL.xml
@@ -167,7 +167,6 @@ android {
 			proguardFile "proguard/app_update.pgcfg"
 			proguardFile "proguard/gms_task.pgcfg"
 			proguardFile "proguard/integrity.pgcfg"
-			proguardFile "proguard/proguard.pgcfg"
 			proguardFile "proguard/review.pgcfg"
 			proguardFile "proguard/split_install.pgcfg"
 		}
